ray marching material refactor
- add template source code NiagaraSimulationStages plugin
- refactor ray marching loop into hlsl macro call
- refactor and build all materials that use ray marching to use macro
- update systems and emitters to properly use new materials
- remove old materials and unused assets
- remove copy of debug slice renderer
- fix debug slice renderer debug bounds default
- debug slice renderer doesn't use alpha from grid

diff --git a/Engine/Plugins/FX/NiagaraSimulationStages/NiagaraSimulationStages.uplugin b/Engine/Plugins/FX/NiagaraSimulationStages/NiagaraSimulationStages.uplugin
--- a/Engine/Plugins/FX/NiagaraSimulationStages/NiagaraSimulationStages.uplugin
+++ b/Engine/Plugins/FX/NiagaraSimulationStages/NiagaraSimulationStages.uplugin
@@ -14,9 +14,16 @@
 	"IsBetaVersion": false,
 	"IsExperimentalVersion": true,
 	"Installed": false,
-    "Plugins": [
+	"Modules": [
+		{
+			"Name": "NiagaraSimulationStages",
+			"Type": "Runtime",
+			"LoadingPhase": "Default"
+		}
+	],    
+	"Plugins": [
         {
-            "Name": "Niagara",
+            "Name": "Niagara",			
             "Enabled": true
         }
 	]
diff --git a/Engine/Plugins/FX/NiagaraSimulationStages/Shaders/NiagaraSimulationStagesRendering.ush b/Engine/Plugins/FX/NiagaraSimulationStages/Shaders/NiagaraSimulationStagesRendering.ush
new file mode 100644
--- /dev/null
+++ b/Engine/Plugins/FX/NiagaraSimulationStages/Shaders/NiagaraSimulationStagesRendering.ush
@@ -0,0 +1,33 @@
+// Copyright Epic Games, Inc. All Rights Reserved.
+
+/*=============================================================================
+NiagaraDataInterfacePhysicsField.ush 
+=============================================================================*/
+
+#pragma once
+
+/* -----------------------------------------------------------------
+ * Rendering code for simulation stages
+ * -----------------------------------------------------------------
+ */
+
+
+#define MARCH_SINGLE_RAY(RayStart, RayDir, UnitStepSize, LocalStepSize, NumSteps, Volume, RetV) \
+RetV = float4(0,0,0,0); \
+float Transmittance = 1; \
+int Steps = ceil(NumSteps); \
+if(Steps >= 1.0 ) Steps += 1.0; \
+[loop] \
+for (int i = 0; i < Steps && Transmittance > 1e-5; i++) \
+{ \
+float Extinction; \
+float Scattering; \
+float3 Luminance; \
+Volume.ComputeStep(RayStart, RayDir, Extinction, Scattering, Luminance); \
+float StepTransmittance = exp(-Extinction * LocalStepSize); \
+Transmittance *= StepTransmittance; \
+RetV.rgb += Luminance * Transmittance * Scattering * LocalStepSize; \
+if(i == Steps-2) UnitStepSize, LocalStepSize *= frac(NumSteps); \
+RayStart += RayDir * UnitStepSize; \
+} \
+RetV.a = 1. - Transmittance; \
diff --git a/Engine/Plugins/FX/NiagaraSimulationStages/Source/NiagaraSimulationStages/NiagaraSimulationStages.Build.cs b/Engine/Plugins/FX/NiagaraSimulationStages/Source/NiagaraSimulationStages/NiagaraSimulationStages.Build.cs
new file mode 100644
--- /dev/null
+++ b/Engine/Plugins/FX/NiagaraSimulationStages/Source/NiagaraSimulationStages/NiagaraSimulationStages.Build.cs
@@ -0,0 +1,52 @@
+// Copyright Epic Games, Inc. All Rights Reserved.
+
+namespace UnrealBuildTool.Rules
+{
+	public class NiagaraSimulationStages : ModuleRules
+	{
+		public NiagaraSimulationStages(ReadOnlyTargetRules Target) : base(Target)
+		{
+			PublicIncludePaths.AddRange(
+				new string[] {
+					// ... add public include paths required here ...
+				}
+				);
+
+			PrivateIncludePaths.AddRange(
+				new string[] {
+                    "../../../../Source/Runtime/Engine/",
+                    "../../../../Source/Runtime/RHI/",
+					// ... add other private include paths required here ...
+				}
+				);
+            
+            PublicDependencyModuleNames.AddRange(
+				new string[]
+				{
+                    "Core",
+                    "CoreUObject",
+                    "Engine",
+                    "RHI",
+                    "Engine",
+                    "RenderCore",
+					// ... add other public dependencies that you statically link with here ...
+				}
+				);
+
+			PrivateDependencyModuleNames.AddRange(
+				new string[]
+                {
+                    "Projects",
+					// ... add private dependencies that you statically link with here ...
+				}
+				);
+
+			DynamicallyLoadedModuleNames.AddRange(
+				new string[]
+				{
+					// ... add any modules that your module loads dynamically here ...
+				}
+				);
+		}
+	}
+}
diff --git a/Engine/Plugins/FX/NiagaraSimulationStages/Source/NiagaraSimulationStages/Private/NiagaraSimulationStages.cpp b/Engine/Plugins/FX/NiagaraSimulationStages/Source/NiagaraSimulationStages/Private/NiagaraSimulationStages.cpp
new file mode 100644
--- /dev/null
+++ b/Engine/Plugins/FX/NiagaraSimulationStages/Source/NiagaraSimulationStages/Private/NiagaraSimulationStages.cpp
@@ -0,0 +1,32 @@
+// Copyright Epic Games, Inc. All Rights Reserved.
+
+#include "CoreMinimal.h"
+#include "Modules/ModuleManager.h"
+#include "INiagaraSimulationStages.h"
+#include "Interfaces/IPluginManager.h"
+#include "Misc/Paths.h"
+#include "ShaderCore.h"
+
+
+class FNiagaraSimulationStages : public INiagaraSimulationStages
+{
+	/** IModuleInterface implementation */
+	virtual void StartupModule() override;
+	virtual void ShutdownModule() override;
+};
+
+IMPLEMENT_MODULE(FNiagaraSimulationStages, NiagaraSimulationStages )
+
+
+void FNiagaraSimulationStages::StartupModule()
+{
+	FString PluginShaderDir = FPaths::Combine(IPluginManager::Get().FindPlugin(TEXT("NiagaraSimulationStages"))->GetBaseDir(), TEXT("Shaders"));
+	AddShaderSourceDirectoryMapping(TEXT("/Plugin/FX/NiagaraSimulationStages"), PluginShaderDir);
+}
+
+
+void FNiagaraSimulationStages::ShutdownModule()
+{
+	// This function may be called during shutdown to clean up your module.  For modules that support dynamic reloading,
+	// we call this function before unloading the module.
+}
diff --git a/Engine/Plugins/FX/NiagaraSimulationStages/Source/NiagaraSimulationStages/Public/INiagaraSimulationStages.h b/Engine/Plugins/FX/NiagaraSimulationStages/Source/NiagaraSimulationStages/Public/INiagaraSimulationStages.h
new file mode 100644
--- /dev/null
+++ b/Engine/Plugins/FX/NiagaraSimulationStages/Source/NiagaraSimulationStages/Public/INiagaraSimulationStages.h
@@ -0,0 +1,39 @@
+// Copyright Epic Games, Inc. All Rights Reserved.
+
+#pragma once
+
+#include "CoreMinimal.h"
+#include "Modules/ModuleInterface.h"
+#include "Modules/ModuleManager.h"
+
+
+/**
+ * The public interface to this module
+ */
+class INiagaraSimulationStages : public IModuleInterface
+{
+
+public:
+
+	/**
+	 * Singleton-like access to this module's interface.  This is just for convenience!
+	 * Beware of calling this during the shutdown phase, though.  Your module might have been unloaded already.
+	 *
+	 * @return Returns singleton instance, loading the module on demand if needed
+	 */
+	static inline INiagaraSimulationStages& Get()
+	{
+		return FModuleManager::LoadModuleChecked< INiagaraSimulationStages >( "NiagaraSimulationStages" );
+	}
+
+	/**
+	 * Checks to see if this module is loaded and ready.  It is only valid to call Get() if IsAvailable() returns true.
+	 *
+	 * @return True if the module is loaded and ready to use
+	 */
+	static inline bool IsAvailable()
+	{
+		return FModuleManager::Get().IsModuleLoaded( "NiagaraSimulationStages" );
+	}
+};
+
