rework ray marching integration using analytical integral for luminance contribution

diff --git a/Engine/Plugins/FX/NiagaraSimulationStages/Shaders/NiagaraSimulationStagesRendering.ush b/Engine/Plugins/FX/NiagaraSimulationStages/Shaders/NiagaraSimulationStagesRendering.ush
--- a/Engine/Plugins/FX/NiagaraSimulationStages/Shaders/NiagaraSimulationStagesRendering.ush
+++ b/Engine/Plugins/FX/NiagaraSimulationStages/Shaders/NiagaraSimulationStagesRendering.ush
@@ -24,9 +24,11 @@ float Extinction; \
 float Scattering; \
 float3 Luminance; \
 Volume.ComputeStep(RayStart, RayDir, Extinction, Scattering, Luminance); \
+Luminance *= Scattering; \
+float ClampedExtinction = max(Extinction, 1e-6); \
 float StepTransmittance = exp(-Extinction * LocalStepSize); \
+RetV.rgb += Transmittance * (Luminance - Luminance * StepTransmittance) / ClampedExtinction; \
 Transmittance *= StepTransmittance; \
-RetV.rgb += Luminance * Transmittance * Scattering * LocalStepSize; \
 if(i == Steps-2) UnitStepSize, LocalStepSize *= frac(NumSteps); \
 RayStart += RayDir * UnitStepSize; \
 } \
